Updated first include of Plugin Templates as Private PCH files were deleted

diff --git a/Engine/Plugins/Editor/PluginBrowser/Templates/Advanced/Source/PLUGIN_NAME/Private/PLUGIN_NAME.cpp b/Engine/Plugins/Editor/PluginBrowser/Templates/Advanced/Source/PLUGIN_NAME/Private/PLUGIN_NAME.cpp
--- a/Engine/Plugins/Editor/PluginBrowser/Templates/Advanced/Source/PLUGIN_NAME/Private/PLUGIN_NAME.cpp
+++ b/Engine/Plugins/Editor/PluginBrowser/Templates/Advanced/Source/PLUGIN_NAME/Private/PLUGIN_NAME.cpp
@@ -1,6 +1,6 @@
 // Copyright 1998-2017 Epic Games, Inc. All Rights Reserved.
 
-#include "PLUGIN_NAMEPrivatePCH.h"
+#include "PLUGIN_NAME.h"
 
 #include "SlateBasics.h"
 #include "SlateExtras.h"
diff --git a/Engine/Plugins/Editor/PluginBrowser/Templates/Advanced/Source/PLUGIN_NAME/Private/PLUGIN_NAMEStyle.cpp b/Engine/Plugins/Editor/PluginBrowser/Templates/Advanced/Source/PLUGIN_NAME/Private/PLUGIN_NAMEStyle.cpp
--- a/Engine/Plugins/Editor/PluginBrowser/Templates/Advanced/Source/PLUGIN_NAME/Private/PLUGIN_NAMEStyle.cpp
+++ b/Engine/Plugins/Editor/PluginBrowser/Templates/Advanced/Source/PLUGIN_NAME/Private/PLUGIN_NAMEStyle.cpp
@@ -1,6 +1,6 @@
 // Copyright 1998-2017 Epic Games, Inc. All Rights Reserved.
 
-#include "PLUGIN_NAMEPrivatePCH.h"
+#include "PLUGIN_NAME.h"
 #include "PLUGIN_NAMEStyle.h"
 #include "SlateGameResources.h"
 #include "IPluginManager.h"
diff --git a/Engine/Plugins/Editor/PluginBrowser/Templates/Advanced/Source/PLUGIN_NAME/Public/PLUGIN_NAMECommands.cpp b/Engine/Plugins/Editor/PluginBrowser/Templates/Advanced/Source/PLUGIN_NAME/Public/PLUGIN_NAMECommands.cpp
--- a/Engine/Plugins/Editor/PluginBrowser/Templates/Advanced/Source/PLUGIN_NAME/Public/PLUGIN_NAMECommands.cpp
+++ b/Engine/Plugins/Editor/PluginBrowser/Templates/Advanced/Source/PLUGIN_NAME/Public/PLUGIN_NAMECommands.cpp
@@ -1,6 +1,6 @@
 // Copyright 1998-2017 Epic Games, Inc. All Rights Reserved.
 
-#include "PLUGIN_NAMEPrivatePCH.h"
+#include "PLUGIN_NAME.h"
 #include "PLUGIN_NAMECommands.h"
 
 #define LOCTEXT_NAMESPACE "FPLUGIN_NAMEModule"
diff --git a/Engine/Plugins/Editor/PluginBrowser/Templates/Basic/Source/PLUGIN_NAME/Private/PLUGIN_NAME.cpp b/Engine/Plugins/Editor/PluginBrowser/Templates/Basic/Source/PLUGIN_NAME/Private/PLUGIN_NAME.cpp
--- a/Engine/Plugins/Editor/PluginBrowser/Templates/Basic/Source/PLUGIN_NAME/Private/PLUGIN_NAME.cpp
+++ b/Engine/Plugins/Editor/PluginBrowser/Templates/Basic/Source/PLUGIN_NAME/Private/PLUGIN_NAME.cpp
@@ -1,6 +1,6 @@
 // Copyright 1998-2017 Epic Games, Inc. All Rights Reserved.
 
-#include "PLUGIN_NAMEPrivatePCH.h"
+#include "PLUGIN_NAME.h"
 
 #include "SlateBasics.h"
 #include "SlateExtras.h"
diff --git a/Engine/Plugins/Editor/PluginBrowser/Templates/Basic/Source/PLUGIN_NAME/Private/PLUGIN_NAMEStyle.cpp b/Engine/Plugins/Editor/PluginBrowser/Templates/Basic/Source/PLUGIN_NAME/Private/PLUGIN_NAMEStyle.cpp
--- a/Engine/Plugins/Editor/PluginBrowser/Templates/Basic/Source/PLUGIN_NAME/Private/PLUGIN_NAMEStyle.cpp
+++ b/Engine/Plugins/Editor/PluginBrowser/Templates/Basic/Source/PLUGIN_NAME/Private/PLUGIN_NAMEStyle.cpp
@@ -1,6 +1,6 @@
 // Copyright 1998-2017 Epic Games, Inc. All Rights Reserved.
 
-#include "PLUGIN_NAMEPrivatePCH.h"
+#include "PLUGIN_NAME.h"
 
 #include "PLUGIN_NAMEStyle.h"
 #include "SlateGameResources.h"
diff --git a/Engine/Plugins/Editor/PluginBrowser/Templates/Basic/Source/PLUGIN_NAME/Public/PLUGIN_NAMECommands.cpp b/Engine/Plugins/Editor/PluginBrowser/Templates/Basic/Source/PLUGIN_NAME/Public/PLUGIN_NAMECommands.cpp
--- a/Engine/Plugins/Editor/PluginBrowser/Templates/Basic/Source/PLUGIN_NAME/Public/PLUGIN_NAMECommands.cpp
+++ b/Engine/Plugins/Editor/PluginBrowser/Templates/Basic/Source/PLUGIN_NAME/Public/PLUGIN_NAMECommands.cpp
@@ -1,6 +1,6 @@
 // Copyright 1998-2017 Epic Games, Inc. All Rights Reserved.
 
-#include "PLUGIN_NAMEPrivatePCH.h"
+#include "PLUGIN_NAME.h"
 #include "PLUGIN_NAMECommands.h"
 
 #define LOCTEXT_NAMESPACE "FPLUGIN_NAMEModule"
diff --git a/Engine/Plugins/Editor/PluginBrowser/Templates/Blank/Source/PLUGIN_NAME/Private/PLUGIN_NAME.cpp b/Engine/Plugins/Editor/PluginBrowser/Templates/Blank/Source/PLUGIN_NAME/Private/PLUGIN_NAME.cpp
--- a/Engine/Plugins/Editor/PluginBrowser/Templates/Blank/Source/PLUGIN_NAME/Private/PLUGIN_NAME.cpp
+++ b/Engine/Plugins/Editor/PluginBrowser/Templates/Blank/Source/PLUGIN_NAME/Private/PLUGIN_NAME.cpp
@@ -1,6 +1,6 @@
 // Copyright 1998-2017 Epic Games, Inc. All Rights Reserved.
 
-#include "PLUGIN_NAMEPrivatePCH.h"
+#include "PLUGIN_NAME.h"
 
 #define LOCTEXT_NAMESPACE "FPLUGIN_NAMEModule"
 
diff --git a/Engine/Plugins/Editor/PluginBrowser/Templates/BlueprintLibrary/Source/PLUGIN_NAME/Private/PLUGIN_NAME.cpp b/Engine/Plugins/Editor/PluginBrowser/Templates/BlueprintLibrary/Source/PLUGIN_NAME/Private/PLUGIN_NAME.cpp
--- a/Engine/Plugins/Editor/PluginBrowser/Templates/BlueprintLibrary/Source/PLUGIN_NAME/Private/PLUGIN_NAME.cpp
+++ b/Engine/Plugins/Editor/PluginBrowser/Templates/BlueprintLibrary/Source/PLUGIN_NAME/Private/PLUGIN_NAME.cpp
@@ -1,6 +1,6 @@
 // Copyright 1998-2017 Epic Games, Inc. All Rights Reserved.
 
-#include "PLUGIN_NAMEPrivatePCH.h"
+#include "PLUGIN_NAME.h"
 
 #define LOCTEXT_NAMESPACE "FPLUGIN_NAMEModule"
 
diff --git a/Engine/Plugins/Editor/PluginBrowser/Templates/BlueprintLibrary/Source/PLUGIN_NAME/Private/PLUGIN_NAMEBPLibrary.cpp b/Engine/Plugins/Editor/PluginBrowser/Templates/BlueprintLibrary/Source/PLUGIN_NAME/Private/PLUGIN_NAMEBPLibrary.cpp
--- a/Engine/Plugins/Editor/PluginBrowser/Templates/BlueprintLibrary/Source/PLUGIN_NAME/Private/PLUGIN_NAMEBPLibrary.cpp
+++ b/Engine/Plugins/Editor/PluginBrowser/Templates/BlueprintLibrary/Source/PLUGIN_NAME/Private/PLUGIN_NAMEBPLibrary.cpp
@@ -1,6 +1,6 @@
 // Copyright 1998-2017 Epic Games, Inc. All Rights Reserved.
 
-#include "PLUGIN_NAMEPrivatePCH.h"
+#include "PLUGIN_NAME.h"
 #include "PLUGIN_NAMEBPLibrary.h"
 
 UPLUGIN_NAMEBPLibrary::UPLUGIN_NAMEBPLibrary(const FObjectInitializer& ObjectInitializer)
diff --git a/Engine/Plugins/Editor/PluginBrowser/Templates/EditorMode/Source/PLUGIN_NAME/Private/PLUGIN_NAME.cpp b/Engine/Plugins/Editor/PluginBrowser/Templates/EditorMode/Source/PLUGIN_NAME/Private/PLUGIN_NAME.cpp
--- a/Engine/Plugins/Editor/PluginBrowser/Templates/EditorMode/Source/PLUGIN_NAME/Private/PLUGIN_NAME.cpp
+++ b/Engine/Plugins/Editor/PluginBrowser/Templates/EditorMode/Source/PLUGIN_NAME/Private/PLUGIN_NAME.cpp
@@ -1,6 +1,6 @@
 // Copyright 1998-2017 Epic Games, Inc. All Rights Reserved.
 
-#include "PLUGIN_NAMEPrivatePCH.h"
+#include "PLUGIN_NAME.h"
 #include "PLUGIN_NAMEEdMode.h"
 
 #define LOCTEXT_NAMESPACE "FPLUGIN_NAMEModule"
diff --git a/Engine/Plugins/Editor/PluginBrowser/Templates/EditorMode/Source/PLUGIN_NAME/Private/PLUGIN_NAMEEdMode.cpp b/Engine/Plugins/Editor/PluginBrowser/Templates/EditorMode/Source/PLUGIN_NAME/Private/PLUGIN_NAMEEdMode.cpp
--- a/Engine/Plugins/Editor/PluginBrowser/Templates/EditorMode/Source/PLUGIN_NAME/Private/PLUGIN_NAMEEdMode.cpp
+++ b/Engine/Plugins/Editor/PluginBrowser/Templates/EditorMode/Source/PLUGIN_NAME/Private/PLUGIN_NAMEEdMode.cpp
@@ -1,6 +1,6 @@
 // Copyright 1998-2017 Epic Games, Inc. All Rights Reserved.
 
-#include "PLUGIN_NAMEPrivatePCH.h"
+#include "PLUGIN_NAME.h"
 #include "PLUGIN_NAMEEdMode.h"
 #include "PLUGIN_NAMEEdModeToolkit.h"
 #include "Toolkits/ToolkitManager.h"
diff --git a/Engine/Plugins/Editor/PluginBrowser/Templates/EditorMode/Source/PLUGIN_NAME/Private/PLUGIN_NAMEEdModeToolkit.cpp b/Engine/Plugins/Editor/PluginBrowser/Templates/EditorMode/Source/PLUGIN_NAME/Private/PLUGIN_NAMEEdModeToolkit.cpp
--- a/Engine/Plugins/Editor/PluginBrowser/Templates/EditorMode/Source/PLUGIN_NAME/Private/PLUGIN_NAMEEdModeToolkit.cpp
+++ b/Engine/Plugins/Editor/PluginBrowser/Templates/EditorMode/Source/PLUGIN_NAME/Private/PLUGIN_NAMEEdModeToolkit.cpp
@@ -1,6 +1,6 @@
 // Copyright 1998-2017 Epic Games, Inc. All Rights Reserved.
 
-#include "PLUGIN_NAMEPrivatePCH.h"
+#include "PLUGIN_NAME.h"
 #include "PLUGIN_NAMEEdMode.h"
 #include "PLUGIN_NAMEEdModeToolkit.h"
 
diff --git a/Engine/Plugins/Editor/PluginBrowser/Templates/ThirdPartyLibrary/Source/PLUGIN_NAME/Private/PLUGIN_NAME.cpp b/Engine/Plugins/Editor/PluginBrowser/Templates/ThirdPartyLibrary/Source/PLUGIN_NAME/Private/PLUGIN_NAME.cpp
--- a/Engine/Plugins/Editor/PluginBrowser/Templates/ThirdPartyLibrary/Source/PLUGIN_NAME/Private/PLUGIN_NAME.cpp
+++ b/Engine/Plugins/Editor/PluginBrowser/Templates/ThirdPartyLibrary/Source/PLUGIN_NAME/Private/PLUGIN_NAME.cpp
@@ -1,6 +1,6 @@
 // Copyright 1998-2017 Epic Games, Inc. All Rights Reserved.
 
-#include "PLUGIN_NAMEPrivatePCH.h"
+#include "PLUGIN_NAME.h"
 #include "Core.h"
 #include "ModuleManager.h"
 #include "IPluginManager.h"
