Remove incorrect include paths from default plugin configuration.

diff --git a/Engine/Plugins/Editor/PluginBrowser/Templates/Advanced/Source/PLUGIN_NAME/PLUGIN_NAME.Build.cs b/Engine/Plugins/Editor/PluginBrowser/Templates/Advanced/Source/PLUGIN_NAME/PLUGIN_NAME.Build.cs
--- a/Engine/Plugins/Editor/PluginBrowser/Templates/Advanced/Source/PLUGIN_NAME/PLUGIN_NAME.Build.cs
+++ b/Engine/Plugins/Editor/PluginBrowser/Templates/Advanced/Source/PLUGIN_NAME/PLUGIN_NAME.Build.cs
@@ -10,7 +10,6 @@ public class PLUGIN_NAME : ModuleRules
 		
 		PublicIncludePaths.AddRange(
 			new string[] {
-				"PLUGIN_NAME/Public"
 				// ... add public include paths required here ...
 			}
 			);
@@ -18,7 +17,6 @@ public class PLUGIN_NAME : ModuleRules
 		
 		PrivateIncludePaths.AddRange(
 			new string[] {
-				"PLUGIN_NAME/Private",
 				// ... add other private include paths required here ...
 			}
 			);
diff --git a/Engine/Plugins/Editor/PluginBrowser/Templates/Basic/Source/PLUGIN_NAME/PLUGIN_NAME.Build.cs b/Engine/Plugins/Editor/PluginBrowser/Templates/Basic/Source/PLUGIN_NAME/PLUGIN_NAME.Build.cs
--- a/Engine/Plugins/Editor/PluginBrowser/Templates/Basic/Source/PLUGIN_NAME/PLUGIN_NAME.Build.cs
+++ b/Engine/Plugins/Editor/PluginBrowser/Templates/Basic/Source/PLUGIN_NAME/PLUGIN_NAME.Build.cs
@@ -10,7 +10,6 @@ public class PLUGIN_NAME : ModuleRules
 		
 		PublicIncludePaths.AddRange(
 			new string[] {
-				"PLUGIN_NAME/Public"
 				// ... add public include paths required here ...
 			}
 			);
@@ -18,7 +17,6 @@ public class PLUGIN_NAME : ModuleRules
 		
 		PrivateIncludePaths.AddRange(
 			new string[] {
-				"PLUGIN_NAME/Private",
 				// ... add other private include paths required here ...
 			}
 			);
diff --git a/Engine/Plugins/Editor/PluginBrowser/Templates/Blank/Source/PLUGIN_NAME/PLUGIN_NAME.Build.cs b/Engine/Plugins/Editor/PluginBrowser/Templates/Blank/Source/PLUGIN_NAME/PLUGIN_NAME.Build.cs
--- a/Engine/Plugins/Editor/PluginBrowser/Templates/Blank/Source/PLUGIN_NAME/PLUGIN_NAME.Build.cs
+++ b/Engine/Plugins/Editor/PluginBrowser/Templates/Blank/Source/PLUGIN_NAME/PLUGIN_NAME.Build.cs
@@ -10,7 +10,6 @@ public class PLUGIN_NAME : ModuleRules
 		
 		PublicIncludePaths.AddRange(
 			new string[] {
-				"PLUGIN_NAME/Public"
 				// ... add public include paths required here ...
 			}
 			);
@@ -18,7 +17,6 @@ public class PLUGIN_NAME : ModuleRules
 		
 		PrivateIncludePaths.AddRange(
 			new string[] {
-				"PLUGIN_NAME/Private",
 				// ... add other private include paths required here ...
 			}
 			);
diff --git a/Engine/Plugins/Editor/PluginBrowser/Templates/BlueprintLibrary/Source/PLUGIN_NAME/PLUGIN_NAME.Build.cs b/Engine/Plugins/Editor/PluginBrowser/Templates/BlueprintLibrary/Source/PLUGIN_NAME/PLUGIN_NAME.Build.cs
--- a/Engine/Plugins/Editor/PluginBrowser/Templates/BlueprintLibrary/Source/PLUGIN_NAME/PLUGIN_NAME.Build.cs
+++ b/Engine/Plugins/Editor/PluginBrowser/Templates/BlueprintLibrary/Source/PLUGIN_NAME/PLUGIN_NAME.Build.cs
@@ -10,8 +10,6 @@ public class PLUGIN_NAME : ModuleRules
 		
 		PublicIncludePaths.AddRange(
 			new string[] {
-				"PLUGIN_NAME/Public"
-				
 				// ... add public include paths required here ...
 			}
 			);
@@ -19,8 +17,6 @@ public class PLUGIN_NAME : ModuleRules
 		
 		PrivateIncludePaths.AddRange(
 			new string[] {
-				"PLUGIN_NAME/Private",
-				
 				// ... add other private include paths required here ...
 			}
 			);
@@ -30,7 +26,6 @@ public class PLUGIN_NAME : ModuleRules
 			new string[]
 			{
 				"Core",
-				
 				// ... add other public dependencies that you statically link with here ...
 			}
 			);
@@ -43,7 +38,6 @@ public class PLUGIN_NAME : ModuleRules
 				"Engine",
 				"Slate",
 				"SlateCore",
-
 				// ... add private dependencies that you statically link with here ...	
 			}
 			);
diff --git a/Engine/Plugins/Editor/PluginBrowser/Templates/EditorMode/Source/PLUGIN_NAME/PLUGIN_NAME.Build.cs b/Engine/Plugins/Editor/PluginBrowser/Templates/EditorMode/Source/PLUGIN_NAME/PLUGIN_NAME.Build.cs
--- a/Engine/Plugins/Editor/PluginBrowser/Templates/EditorMode/Source/PLUGIN_NAME/PLUGIN_NAME.Build.cs
+++ b/Engine/Plugins/Editor/PluginBrowser/Templates/EditorMode/Source/PLUGIN_NAME/PLUGIN_NAME.Build.cs
@@ -10,7 +10,6 @@ public class PLUGIN_NAME : ModuleRules
 		
 		PublicIncludePaths.AddRange(
 			new string[] {
-				"PLUGIN_NAME/Public"
 				// ... add public include paths required here ...
 			}
 			);
@@ -18,7 +17,6 @@ public class PLUGIN_NAME : ModuleRules
 		
 		PrivateIncludePaths.AddRange(
 			new string[] {
-				"PLUGIN_NAME/Private",
 				// ... add other private include paths required here ...
 			}
 			);
diff --git a/Engine/Plugins/Editor/PluginBrowser/Templates/ThirdPartyLibrary/Source/PLUGIN_NAME/PLUGIN_NAME.Build.cs b/Engine/Plugins/Editor/PluginBrowser/Templates/ThirdPartyLibrary/Source/PLUGIN_NAME/PLUGIN_NAME.Build.cs
--- a/Engine/Plugins/Editor/PluginBrowser/Templates/ThirdPartyLibrary/Source/PLUGIN_NAME/PLUGIN_NAME.Build.cs
+++ b/Engine/Plugins/Editor/PluginBrowser/Templates/ThirdPartyLibrary/Source/PLUGIN_NAME/PLUGIN_NAME.Build.cs
@@ -10,7 +10,6 @@ public class PLUGIN_NAME : ModuleRules
 		
 		PublicIncludePaths.AddRange(
 			new string[] {
-				"PLUGIN_NAME/Public"
 				// ... add public include paths required here ...
 			}
 			);
@@ -18,7 +17,6 @@ public class PLUGIN_NAME : ModuleRules
 		
 		PrivateIncludePaths.AddRange(
 			new string[] {
-				"PLUGIN_NAME/Private",
 				// ... add other private include paths required here ...
 			}
 			);
